Fix HTML5 CIS

diff --git a/Engine/Source/Developer/HTML5/HTML5TargetPlatform/Private/HTML5TargetPlatform.h b/Engine/Source/Developer/HTML5/HTML5TargetPlatform/Private/HTML5TargetPlatform.h
--- a/Engine/Source/Developer/HTML5/HTML5TargetPlatform/Private/HTML5TargetPlatform.h
+++ b/Engine/Source/Developer/HTML5/HTML5TargetPlatform/Private/HTML5TargetPlatform.h
@@ -6,6 +6,9 @@
 
 #pragma once
 
+#if WITH_ENGINE
+#include "StaticMeshResources.h"
+#endif // WITH_ENGINE
 
 /**
  * Implements the HTML5 target platform.
