fix HTML5 not supporting File/Package
TTP338976
#ue4
#html5

diff --git a/Engine/Source/Developer/HTML5/HTML5TargetPlatform/Private/HTML5TargetPlatform.h b/Engine/Source/Developer/HTML5/HTML5TargetPlatform/Private/HTML5TargetPlatform.h
--- a/Engine/Source/Developer/HTML5/HTML5TargetPlatform/Private/HTML5TargetPlatform.h
+++ b/Engine/Source/Developer/HTML5/HTML5TargetPlatform/Private/HTML5TargetPlatform.h
@@ -44,6 +44,16 @@ public:
 
 	virtual bool IsRunningPlatform( ) const override;
 
+	virtual bool SupportsFeature( ETargetPlatformFeatures::Type Feature ) const override
+	{
+		if (Feature == ETargetPlatformFeatures::Packaging)
+		{
+			return true;
+		}
+
+		return TTargetPlatformBase<FHTML5PlatformProperties>::SupportsFeature(Feature);
+	}
+
 #if WITH_ENGINE
 	virtual void GetAllPossibleShaderFormats( TArray<FName>& OutFormats ) const override;
 
