implement IsSDKInstalled on HTML5 target platform
#ue4
#html5

diff --git a/Engine/Source/Developer/HTML5/HTML5TargetPlatform/Private/HTML5TargetPlatform.cpp b/Engine/Source/Developer/HTML5/HTML5TargetPlatform/Private/HTML5TargetPlatform.cpp
--- a/Engine/Source/Developer/HTML5/HTML5TargetPlatform/Private/HTML5TargetPlatform.cpp
+++ b/Engine/Source/Developer/HTML5/HTML5TargetPlatform/Private/HTML5TargetPlatform.cpp
@@ -60,6 +60,15 @@ ITargetDevicePtr FHTML5TargetPlatform::GetDevice( const FTargetDeviceId& DeviceI
 	return NULL;
 }
 
+bool FHTML5TargetPlatform::IsSdkInstalled(bool bProjectHasCode, FString& OutDocumentationPath) const
+{
+	bool bHTML5SDKInstalled = false; // @todo How do we check that the iOS SDK is installed when building from Windows? Is that even possible?
+	TCHAR BaseSDKPath[256];
+	FPlatformMisc::GetEnvironmentVariable(TEXT("EMSCRIPTEN"), BaseSDKPath, ARRAY_COUNT(BaseSDKPath));
+	bHTML5SDKInstalled = FString(BaseSDKPath).Len() > 0;
+	return bHTML5SDKInstalled;
+}
+
 
 bool FHTML5TargetPlatform::IsRunningPlatform( ) const
 {
diff --git a/Engine/Source/Developer/HTML5/HTML5TargetPlatform/Private/HTML5TargetPlatform.h b/Engine/Source/Developer/HTML5/HTML5TargetPlatform/Private/HTML5TargetPlatform.h
--- a/Engine/Source/Developer/HTML5/HTML5TargetPlatform/Private/HTML5TargetPlatform.h
+++ b/Engine/Source/Developer/HTML5/HTML5TargetPlatform/Private/HTML5TargetPlatform.h
@@ -54,6 +54,8 @@ public:
 		return TTargetPlatformBase<FHTML5PlatformProperties>::SupportsFeature(Feature);
 	}
 
+	virtual bool IsSdkInstalled(bool bProjectHasCode, FString& OutDocumentationPath) const override;
+
 #if WITH_ENGINE
 	virtual void GetAllPossibleShaderFormats( TArray<FName>& OutFormats ) const override;
 
