make it so the HTML5 platform returns null if the SDK is not installed
#ue4
#html5

diff --git a/Engine/Source/Developer/HTML5/HTML5TargetPlatform/Private/HTML5TargetPlatformModule.cpp b/Engine/Source/Developer/HTML5/HTML5TargetPlatform/Private/HTML5TargetPlatformModule.cpp
--- a/Engine/Source/Developer/HTML5/HTML5TargetPlatform/Private/HTML5TargetPlatformModule.cpp
+++ b/Engine/Source/Developer/HTML5/HTML5TargetPlatform/Private/HTML5TargetPlatformModule.cpp
@@ -40,6 +40,12 @@ public:
 		{
 			// finally, make the interface object
 			Singleton = new FHTML5TargetPlatform();
+			FString OutPath;
+			if (!Singleton->IsSdkInstalled(false, OutPath))
+			{
+				delete Singleton;
+				Singleton = NULL;
+			}
 		}
 
 		return Singleton;
