Fix flipped b/r. #html5

 #ttp 345097 HTML5: SideScroller2D: Background is orange

diff --git a/Engine/Source/Developer/HTML5/HTML5TargetPlatform/Private/HTML5TargetPlatform.cpp b/Engine/Source/Developer/HTML5/HTML5TargetPlatform/Private/HTML5TargetPlatform.cpp
--- a/Engine/Source/Developer/HTML5/HTML5TargetPlatform/Private/HTML5TargetPlatform.cpp
+++ b/Engine/Source/Developer/HTML5/HTML5TargetPlatform/Private/HTML5TargetPlatform.cpp
@@ -163,7 +163,11 @@ void FHTML5TargetPlatform::GetTextureFormats( const UTexture* Texture, TArray<FN
 			{
 				TextureFormatName = NameG8;
 			}
-			else
+			else if ( Texture->CompressionSettings == TC_EditorIcon )
+			{
+				TextureFormatName = NameRGBA8; 
+			} 
+			else 
 			{
 				TextureFormatName = NameBGRA8;
 			}
