Make HTML5 Device name unique to machine. #ttp347362
Prevents devices in the Device Manager being claimed by other machines on the same network.

diff --git a/Engine/Source/Developer/HTML5/HTML5TargetPlatform/Private/HTML5TargetPlatform.cpp b/Engine/Source/Developer/HTML5/HTML5TargetPlatform/Private/HTML5TargetPlatform.cpp
--- a/Engine/Source/Developer/HTML5/HTML5TargetPlatform/Private/HTML5TargetPlatform.cpp
+++ b/Engine/Source/Developer/HTML5/HTML5TargetPlatform/Private/HTML5TargetPlatform.cpp
@@ -34,7 +34,7 @@ FHTML5TargetPlatform::FHTML5TargetPlatform( )
 		const FString& BrowserName = It.Key.ToString();
 		const FString& BrowserPath = It.Value;
 		if(FPlatformFileManager::Get().GetPlatformFile().FileExists(*It.Value))
-			LocalDevice.Add(MakeShareable(new FHTML5TargetDevice(*this, It.Key.ToString())));
+			LocalDevice.Add(MakeShareable(new FHTML5TargetDevice(*this, FString::Printf(TEXT("%s on %s"), *It.Key.ToString(), FPlatformProcess::ComputerName()))));
 	}
 
 #if WITH_ENGINE
