Fix for compiling UFE w/HTML5.

diff --git a/Engine/Source/Developer/HTML5/HTML5TargetPlatform/Private/HTML5TargetPlatformPrivatePCH.h b/Engine/Source/Developer/HTML5/HTML5TargetPlatform/Private/HTML5TargetPlatformPrivatePCH.h
--- a/Engine/Source/Developer/HTML5/HTML5TargetPlatform/Private/HTML5TargetPlatformPrivatePCH.h
+++ b/Engine/Source/Developer/HTML5/HTML5TargetPlatform/Private/HTML5TargetPlatformPrivatePCH.h
@@ -7,6 +7,7 @@
  *****************************************************************************/
 
 #include "Core.h"
+#include "CoreUObject.h"
 #include "Runtime/Core/Public/HTML5/HTML5PlatformProperties.h"
 
 #if WITH_ENGINE
