Crash fix for missing html5 ini values.

diff --git a/Engine/Source/Developer/HTML5/HTML5TargetPlatform/Private/HTML5TargetPlatform.cpp b/Engine/Source/Developer/HTML5/HTML5TargetPlatform/Private/HTML5TargetPlatform.cpp
--- a/Engine/Source/Developer/HTML5/HTML5TargetPlatform/Private/HTML5TargetPlatform.cpp
+++ b/Engine/Source/Developer/HTML5/HTML5TargetPlatform/Private/HTML5TargetPlatform.cpp
@@ -65,46 +65,53 @@ ITargetDevicePtr FHTML5TargetPlatform::GetDevice( const FTargetDeviceId& DeviceI
 
 bool FHTML5TargetPlatform::IsSdkInstalled(bool bProjectHasCode, FString& OutDocumentationPath) const
 {
+	FConfigSection SDKPaths;
 	//New style detection of devices
-	FConfigSection SDKPaths = HTML5EngineSettings["/Script/HTML5PlatformEditor.HTML5SDKSettings"];
-	for (auto It : SDKPaths)
+	if (HTML5EngineSettings.Find("/Script/HTML5PlatformEditor.HTML5SDKSettings"))
 	{
-		const FString& Platform = It.Key.ToString();
-		FString Path = It.Value;
-		{	
-			Path.RemoveFromStart(TEXT("(Path=\""));
-			Path.RemoveFromEnd(TEXT("\")"));
-			if (Platform == "Emscripten" && IFileManager::Get().DirectoryExists(*Path))
+		SDKPaths = HTML5EngineSettings["/Script/HTML5PlatformEditor.HTML5SDKSettings"];
+		for (auto It : SDKPaths)
+		{
+			const FString& Platform = It.Key.ToString();
+			FString Path = It.Value;
 			{
-				return true;
+				Path.RemoveFromStart(TEXT("(Path=\""));
+				Path.RemoveFromEnd(TEXT("\")"));
+				if (Platform == "Emscripten" && IFileManager::Get().DirectoryExists(*Path))
+				{
+					return true;
+				}
 			}
 		}
 	}
 
 	// Old fallbacks
-	FString SectionName = "HTML5SDKPaths";
-	SDKPaths = HTML5EngineSettings[SectionName];
-	for (auto It : SDKPaths )
+	if (HTML5EngineSettings.Find("HTML5SDKPaths"))
 	{
-		const FString& Platform = It.Key.ToString();
-		const FString& Path = It.Value;
+		FString SectionName = "HTML5SDKPaths";
+		SDKPaths = HTML5EngineSettings[SectionName];
+		for (auto It : SDKPaths)
 		{
-			if (Platform == "Emscripten" && IFileManager::Get().DirectoryExists(*Path))
+			const FString& Platform = It.Key.ToString();
+			const FString& Path = It.Value;
 			{
-				return true;
-			}
+				if (Platform == "Emscripten" && IFileManager::Get().DirectoryExists(*Path))
+				{
+					return true;
+				}
 #if PLATFORM_WINDOWS
-			if ( Platform == "Windows" && IFileManager::Get().DirectoryExists(*Path)) 
-			{
-				return true; 
-			}
+				if (Platform == "Windows" && IFileManager::Get().DirectoryExists(*Path))
+				{
+					return true;
+				}
 #endif 
 #if PLATFORM_MAC
-			if ( Platform == "Mac" && IFileManager::Get().DirectoryExists(*Path)) 
-			{
-				return true; 
-			}
+				if (Platform == "Mac" && IFileManager::Get().DirectoryExists(*Path))
+				{
+					return true;
+				}
 #endif 
+			}
 		}
 	}
 
@@ -281,38 +288,41 @@ void FHTML5TargetPlatform::RefreshAvailableDevices()
 		Config = &HTML5EngineSettings;
 	}
 	TArray<FString> ValueArray;
-	FConfigSection AvaliableDevicesNewSection =  (*Config)["/Script/HTML5PlatformEditor.HTML5SDKSettings"];
-	for (auto It : AvaliableDevicesNewSection)
+	if (Config->Find("/Script/HTML5PlatformEditor.HTML5SDKSettings")) 
 	{
-		ValueArray.Reset();
-		if (It.Key == TEXT("DeviceMap"))
+		FConfigSection AvaliableDevicesNewSection =  (*Config)["/Script/HTML5PlatformEditor.HTML5SDKSettings"];
+		for (auto It : AvaliableDevicesNewSection)
 		{
-			FString DeviceName;
-			FString DevicePath;
-			It.Value.RemoveFromStart(TEXT("("));
-			It.Value.RemoveFromEnd(TEXT(")"));
-			It.Value.ParseIntoArray(&ValueArray, TEXT(","), 1);
-			for (auto& Value : ValueArray)
+			ValueArray.Reset();
+			if (It.Key == TEXT("DeviceMap"))
 			{
-				if (Value.StartsWith(TEXT("DeviceName=")))
+				FString DeviceName;
+				FString DevicePath;
+				It.Value.RemoveFromStart(TEXT("("));
+				It.Value.RemoveFromEnd(TEXT(")"));
+				It.Value.ParseIntoArray(&ValueArray, TEXT(","), 1);
+				for (auto& Value : ValueArray)
 				{
-					DeviceName = Value.RightChop(11).TrimQuotes();
+					if (Value.StartsWith(TEXT("DeviceName=")))
+					{
+						DeviceName = Value.RightChop(11).TrimQuotes();
+					}
+					else if (Value.StartsWith(TEXT("DevicePath=(FilePath=")))
+					{
+						DevicePath = Value.RightChop(21);
+						DevicePath.RemoveFromEnd(TEXT(")"));
+						DevicePath = DevicePath.TrimQuotes();
+					}
 				}
-				else if (Value.StartsWith(TEXT("DevicePath=(FilePath=")))
+
+				if (!DeviceName.IsEmpty() && !DevicePath.IsEmpty() &&
+					FPlatformFileManager::Get().GetPlatformFile().FileExists(*DevicePath))
 				{
-					DevicePath = Value.RightChop(21);
-					DevicePath.RemoveFromEnd(TEXT(")"));
-					DevicePath = DevicePath.TrimQuotes();
+					ITargetDevicePtr Device = MakeShareable(new FHTML5TargetDevice(*this, FString::Printf(TEXT("%s on %s"), *DeviceName, FPlatformProcess::ComputerName())));
+					LocalDevice.Add(Device);
+					DeviceDiscoveredEvent.Broadcast(Device.ToSharedRef());
 				}
 			}
-
-			if (!DeviceName.IsEmpty() && !DevicePath.IsEmpty() &&
-				FPlatformFileManager::Get().GetPlatformFile().FileExists(*DevicePath))
-			{
-				ITargetDevicePtr Device = MakeShareable(new FHTML5TargetDevice(*this, FString::Printf(TEXT("%s on %s"), *DeviceName, FPlatformProcess::ComputerName())));
-				LocalDevice.Add(Device);
-				DeviceDiscoveredEvent.Broadcast(Device.ToSharedRef());
-			}
 		}
 	}
 
@@ -328,17 +338,19 @@ void FHTML5TargetPlatform::RefreshAvailableDevices()
 #else 
 		DeviceSectionName = "HTML5DevicesLinux";
 #endif 
-
-		FConfigSection AvaliableDevicesSection = HTML5EngineSettings[DeviceSectionName];
-		for (auto It : AvaliableDevicesSection)
+		if (HTML5EngineSettings.Find(DeviceSectionName))
 		{
-			const FString& BrowserName = It.Key.ToString();
-			const FString& BrowserPath = It.Value;
-			if (FPlatformFileManager::Get().GetPlatformFile().FileExists(*It.Value))
+			FConfigSection AvaliableDevicesSection = HTML5EngineSettings[DeviceSectionName];
+			for (auto It : AvaliableDevicesSection)
 			{
-				ITargetDevicePtr Device = MakeShareable(new FHTML5TargetDevice(*this, FString::Printf(TEXT("%s on %s"), *It.Key.ToString(), FPlatformProcess::ComputerName())));
-				LocalDevice.Add(Device);
-				DeviceDiscoveredEvent.Broadcast(Device.ToSharedRef());
+				const FString& BrowserName = It.Key.ToString();
+				const FString& BrowserPath = It.Value;
+				if (FPlatformFileManager::Get().GetPlatformFile().FileExists(*It.Value))
+				{
+					ITargetDevicePtr Device = MakeShareable(new FHTML5TargetDevice(*this, FString::Printf(TEXT("%s on %s"), *It.Key.ToString(), FPlatformProcess::ComputerName())));
+					LocalDevice.Add(Device);
+					DeviceDiscoveredEvent.Broadcast(Device.ToSharedRef());
+				}
 			}
 		}
 	}
