fix for editor crash on exit
#html5

diff --git a/Engine/Source/Developer/HTML5/HTML5PlatformEditor/Private/HTML5PlatformEditorModule.cpp b/Engine/Source/Developer/HTML5/HTML5PlatformEditor/Private/HTML5PlatformEditorModule.cpp
--- a/Engine/Source/Developer/HTML5/HTML5PlatformEditor/Private/HTML5PlatformEditorModule.cpp
+++ b/Engine/Source/Developer/HTML5/HTML5PlatformEditor/Private/HTML5PlatformEditorModule.cpp
@@ -47,8 +47,6 @@ class FHTML5PlatformEditorModule
 
 	virtual void ShutdownModule() override
 	{
-		auto* Settings = GetMutableDefault<UHTML5SDKSettings>();
-
 		ISettingsModule* SettingsModule = FModuleManager::GetModulePtr<ISettingsModule>("Settings");
 
 		if (SettingsModule != nullptr)
