Fixed typo with default device detection.
Allow Fallback onto EMSCRIPTEN env var when SDK settings are blank.

diff --git a/Engine/Source/Developer/HTML5/HTML5PlatformEditor/Private/HTML5SDKSettings.cpp b/Engine/Source/Developer/HTML5/HTML5PlatformEditor/Private/HTML5SDKSettings.cpp
--- a/Engine/Source/Developer/HTML5/HTML5PlatformEditor/Private/HTML5SDKSettings.cpp
+++ b/Engine/Source/Developer/HTML5/HTML5PlatformEditor/Private/HTML5SDKSettings.cpp
@@ -32,7 +32,7 @@ void UHTML5SDKSettings::PostEditChangeProperty(struct FPropertyChangedEvent& Pro
 	
 	if (TargetPlatformModule) 
 	{
-		SaveConfig();
+		UpdateGlobalUserConfigFile();//SaveConfig();
 		TargetPlatformModule->RefreshAvailableDevices();
 	}
 }
@@ -51,7 +51,8 @@ void UHTML5SDKSettings::QueryKnownBrowserLocations()
 	} PossibleLocations[] = 
 	{
 #if PLATFORM_WINDOWS
-		{TEXT("Nightly"), TEXT("C:/Program Files/Nightly/firefox.exe")},
+		{TEXT("Nightly(64bit)"), TEXT("C:/Program Files/Nightly/firefox.exe")},
+		{TEXT("Nightly"), TEXT("C:/Program Files (x86)/Nightly/firefox.exe")},
 		{TEXT("Firefox"), TEXT("C:/Program Files (x86)/Mozilla Firefox/firefox.exe")},
 		{TEXT("Chrome"), TEXT("C:/Program Files (x86)/Google/Chrome/Application/chrome.exe")},
 #elif PLATFORM_MAC
@@ -65,16 +66,23 @@ void UHTML5SDKSettings::QueryKnownBrowserLocations()
 #endif
 	};
 
+	bool bDirty = false;
 	for (const auto& Loc : PossibleLocations)
 	{
-		if (FPlatformFileManager::Get().GetPlatformFile().FileExists(*Loc.Name))
+		if (FPlatformFileManager::Get().GetPlatformFile().FileExists(*Loc.Path) || FPlatformFileManager::Get().GetPlatformFile().DirectoryExists(*Loc.Path))
 		{
 			FHTML5DeviceMapping NewDevice;
 			NewDevice.DeviceName = Loc.Name;
 			NewDevice.DevicePath.FilePath = Loc.Path;
 			DeviceMap.Add(NewDevice);
+			bDirty = true;
 		}
 	}
+
+	if (bDirty)
+	{
+		UpdateGlobalUserConfigFile();
+	}
 }
 
 TSharedRef<IPropertyTypeCustomization> FHTML5SDKPathCustomization::MakeInstance()
@@ -240,6 +248,15 @@ void FHTML5SDKPathCustomization::UpdateAvailableVersions()
 	TArray<FHTML5SDKVersionNumber> SDKVersions;
 	TargetPlatformModule->GetInstalledSDKVersions(*(Settings->EmscriptenRoot.SDKPath), SDKVersions);
 
+	if (SDKVersions.Num() == 0) 
+	{
+		//Try the environment variable EMSCRIPTEN - This assumes latest version.
+		TCHAR SDKDirectory[32768] = { 0 };
+		FPlatformMisc::GetEnvironmentVariable(TEXT("EMSCRIPTEN"), SDKDirectory, ARRAY_COUNT(SDKDirectory));
+		TargetPlatformModule->GetInstalledSDKVersions(SDKDirectory, SDKVersions);
+		PrevSelectedVersion->Version.MakeLatestVersionNumber();
+	}
+
 	if (SDKVersions.Num() == 0)
 	{
 		return;
diff --git a/Engine/Source/Developer/HTML5/HTML5PlatformEditor/Private/HTML5SDKSettings.h b/Engine/Source/Developer/HTML5/HTML5PlatformEditor/Private/HTML5SDKSettings.h
--- a/Engine/Source/Developer/HTML5/HTML5PlatformEditor/Private/HTML5SDKSettings.h
+++ b/Engine/Source/Developer/HTML5/HTML5PlatformEditor/Private/HTML5SDKSettings.h
@@ -45,15 +45,16 @@ class HTML5PLATFORMEDITOR_API UHTML5SDKSettings : public UObject
 public:
 	GENERATED_UCLASS_BODY()
 
-	UPROPERTY(GlobalConfig, EditAnywhere, Category = HTML5_SDK_Paths, Meta = (DisplayName = "Location of Python exe (This can be blank if python is on the PATH)"))
+	// Path to the python executable. Can be blank is python is on the PATH
+	UPROPERTY(GlobalConfig, EditAnywhere, Category = HTML5_SDK_Paths, Meta = (DisplayName = "Location of Python exe"))
 	FFilePath Python;
 
 	// Available browsers that can be used when launching HTML5 builds.
 	UPROPERTY(GlobalConfig, EditAnywhere, Category = HTML5_Devices, Meta = (DisplayName = "Available browsers"))
 	TArray<FHTML5DeviceMapping> DeviceMap;
 
-	// Path to Emscripten SDK install directory
-	UPROPERTY(GlobalConfig, EditAnywhere, Category = HTML5_SDK_Paths, Meta = (DisplayName = "Location of Emscripten SDK (This directory contains all emscripten versions)"))
+	// Path to Emscripten SDK install directory. This directory contains emsdk. Will use the environment variable "EMSCRIPTEN" if left blank
+	UPROPERTY(GlobalConfig, EditAnywhere, Category = HTML5_SDK_Paths, Meta = (DisplayName = "Location of Emscripten SDK"))
 	FHTML5SDKPath EmscriptenRoot;
 
 #if WITH_EDITOR
diff --git a/Engine/Source/Programs/UnrealBuildTool/HTML5/HTML5SDKInfo.cs b/Engine/Source/Programs/UnrealBuildTool/HTML5/HTML5SDKInfo.cs
--- a/Engine/Source/Programs/UnrealBuildTool/HTML5/HTML5SDKInfo.cs
+++ b/Engine/Source/Programs/UnrealBuildTool/HTML5/HTML5SDKInfo.cs
@@ -226,6 +226,12 @@ namespace UnrealBuildTool
 					}
 				}
 
+				if (string.IsNullOrEmpty(SDKPathString) && Environment.GetEnvironmentVariable("EMSCRIPTEN") != null)
+				{
+					VersionString = "-1.-1.-1";
+					SDKPathString = Environment.GetEnvironmentVariable("EMSCRIPTEN");
+				}
+
 				if (!string.IsNullOrEmpty(SDKPathString) && !string.IsNullOrEmpty(VersionString))
 				{
 					var SDKVersions = GetInstalledVersions(SDKPathString);
