UE-12561: uncompressed lightmaps rendering incorrectly in HTML5

diff --git a/Engine/Source/Developer/HTML5/HTML5TargetPlatform/Private/HTML5TargetPlatform.cpp b/Engine/Source/Developer/HTML5/HTML5TargetPlatform/Private/HTML5TargetPlatform.cpp
--- a/Engine/Source/Developer/HTML5/HTML5TargetPlatform/Private/HTML5TargetPlatform.cpp
+++ b/Engine/Source/Developer/HTML5/HTML5TargetPlatform/Private/HTML5TargetPlatform.cpp
@@ -246,13 +246,9 @@ void FHTML5TargetPlatform::GetTextureFormats( const UTexture* Texture, TArray<FN
 			{
 				TextureFormatName = NameG8;
 			}
-			else if ( Texture->CompressionSettings == TC_EditorIcon )
-			{
-				TextureFormatName = NameRGBA8; 
-			} 
 			else 
 			{
-				TextureFormatName = NameBGRA8;
+				TextureFormatName = NameRGBA8;
 			}
 	}
 	else if (Texture->CompressionSettings == TC_HDR
