Found HTML5 devices are not saved at startup.

diff --git a/Engine/Source/Developer/HTML5/HTML5PlatformEditor/Private/HTML5SDKSettings.cpp b/Engine/Source/Developer/HTML5/HTML5PlatformEditor/Private/HTML5SDKSettings.cpp
--- a/Engine/Source/Developer/HTML5/HTML5PlatformEditor/Private/HTML5SDKSettings.cpp
+++ b/Engine/Source/Developer/HTML5/HTML5PlatformEditor/Private/HTML5SDKSettings.cpp
@@ -66,7 +66,6 @@ void UHTML5SDKSettings::QueryKnownBrowserLocations()
 #endif
 	};
 
-	bool bDirty = false;
 	for (const auto& Loc : PossibleLocations)
 	{
 		if (FPlatformFileManager::Get().GetPlatformFile().FileExists(*Loc.Path) || FPlatformFileManager::Get().GetPlatformFile().DirectoryExists(*Loc.Path))
@@ -75,14 +74,8 @@ void UHTML5SDKSettings::QueryKnownBrowserLocations()
 			NewDevice.DeviceName = Loc.Name;
 			NewDevice.DevicePath.FilePath = Loc.Path;
 			DeviceMap.Add(NewDevice);
-			bDirty = true;
 		}
 	}
-
-	if (bDirty)
-	{
-		UpdateGlobalUserConfigFile();
-	}
 }
 
 TSharedRef<IPropertyTypeCustomization> FHTML5SDKPathCustomization::MakeInstance()
