UE-16228 Fix up confusing and incorrect tooltip for html5 sdk settings.

diff --git a/Engine/Source/Developer/HTML5/HTML5PlatformEditor/Private/HTML5SDKSettings.h b/Engine/Source/Developer/HTML5/HTML5PlatformEditor/Private/HTML5SDKSettings.h
--- a/Engine/Source/Developer/HTML5/HTML5PlatformEditor/Private/HTML5SDKSettings.h
+++ b/Engine/Source/Developer/HTML5/HTML5PlatformEditor/Private/HTML5SDKSettings.h
@@ -53,7 +53,7 @@ public:
 	UPROPERTY(GlobalConfig, EditAnywhere, Category = HTML5_Devices, Meta = (DisplayName = "Available browsers"))
 	TArray<FHTML5DeviceMapping> DeviceMap;
 
-	// Path to Emscripten SDK install directory. This directory contains emsdk. Will use the environment variable "EMSCRIPTEN" if left blank
+	// Path to Emscripten SDK install directory. This is the directory which contains emsdk.bat and emsdk.sh
 	UPROPERTY(GlobalConfig, EditAnywhere, Category = HTML5_SDK_Paths, Meta = (DisplayName = "Location of Emscripten SDK"))
 	FHTML5SDKPath EmscriptenRoot;
 
