HTML5 - remove experimental HTML5 settings from Project Settings

- offscreen canvas
- LLVM wasm backend

#rnx

diff --git a/Engine/Source/Developer/HTML5/HTML5PlatformEditor/Classes/HTML5TargetSettings.h b/Engine/Source/Developer/HTML5/HTML5PlatformEditor/Classes/HTML5TargetSettings.h
--- a/Engine/Source/Developer/HTML5/HTML5PlatformEditor/Classes/HTML5TargetSettings.h
+++ b/Engine/Source/Developer/HTML5/HTML5PlatformEditor/Classes/HTML5TargetSettings.h
@@ -62,14 +62,14 @@ public:
 	 * Use OffscreenCanvas Support (else use Offscreen Framebuffer)
 	 * Note: Multithreading rendering option (WORK IN PROGRESS)
 	 */
-	UPROPERTY(GlobalConfig, EditAnywhere, Category=Emscripten, Meta = (DisplayName = "Offscreen Canvas [experimental]"))
-	bool OffscreenCanvas;
+//	UPROPERTY(GlobalConfig, EditAnywhere, Category=Emscripten, Meta = (DisplayName = "Offscreen Canvas [experimental]"))
+//	bool OffscreenCanvas;
 
 	/**
 	 * Use LLVM WASM Backend (WORK IN PROGRESS)
 	 */
-	UPROPERTY(GlobalConfig, EditAnywhere, Category=Emscripten, Meta = (DisplayName = "LLVM Wasm backend [experimental]"))
-	bool LLVMWasmBackend;
+//	UPROPERTY(GlobalConfig, EditAnywhere, Category=Emscripten, Meta = (DisplayName = "LLVM Wasm backend [experimental]"))
+//	bool LLVMWasmBackend;
 
 	/**
 	 * Enable Tracing (trace.h)
