#rnx

diff --git a/Engine/Source/Developer/HTML5/HTML5PlatformEditor/Private/HTML5TargetSettingsCustomization.cpp b/Engine/Source/Developer/HTML5/HTML5PlatformEditor/Private/HTML5TargetSettingsCustomization.cpp
--- a/Engine/Source/Developer/HTML5/HTML5PlatformEditor/Private/HTML5TargetSettingsCustomization.cpp
+++ b/Engine/Source/Developer/HTML5/HTML5PlatformEditor/Private/HTML5TargetSettingsCustomization.cpp
@@ -46,6 +46,8 @@ void FHTML5TargetSettingsCustomization::CustomizeDetails(IDetailLayoutBuilder& D
 
 	SETUP_SOURCEONLY_PROP(EnableIndexedDB, EmscriptenCategory);
 
+	SETUP_SOURCEONLY_PROP(EnableMultithreading, EmscriptenCategory);
+
 	AudioPluginWidgetManager.BuildAudioCategory(DetailLayout, EAudioPlatform::HTML5);
 }
 
