HTML5 - deprecation messaging

[RN]

diff --git a/Engine/Source/Developer/HTML5/HTML5PlatformEditor/Private/HTML5TargetSettingsCustomization.cpp b/Engine/Source/Developer/HTML5/HTML5PlatformEditor/Private/HTML5TargetSettingsCustomization.cpp
--- a/Engine/Source/Developer/HTML5/HTML5PlatformEditor/Private/HTML5TargetSettingsCustomization.cpp
+++ b/Engine/Source/Developer/HTML5/HTML5PlatformEditor/Private/HTML5TargetSettingsCustomization.cpp
@@ -8,6 +8,20 @@
 #include "DetailCategoryBuilder.h"
 #include "Misc/EngineBuildSettings.h"
 
+#include "DetailWidgetRow.h"
+//#include "Widgets/DeclarativeSyntaxSupport.h"
+//#include "Widgets/SBoxPanel.h"
+#include "Widgets/SCompoundWidget.h"
+//#include "Widgets/SWidget.h"
+#include "Widgets/Images/SImage.h"
+//#include "Widgets/Input/SButton.h"
+#include "Widgets/Layout/SBorder.h"
+//#include "Widgets/Layout/SBox.h"
+//#include "Widgets/Layout/SWidgetSwitcher.h"
+//#include "Widgets/Notifications/SNotificationList.h"
+//#include "Widgets/Text/SRichTextBlock.h"
+#include "Widgets/Text/STextBlock.h"
+
 #define LOCTEXT_NAMESPACE "HTML5TargetSettings"
 DEFINE_LOG_CATEGORY_STATIC(LogIOSTargetSettings, Log, All);
 
@@ -19,6 +33,60 @@ namespace FHTML5TargetSettingsCustomizationConstants
 }
 
 
+//////////////////////////////////////////////////////////////////////////
+class HTML5DeprecationBanner : public SCompoundWidget
+{
+	SLATE_BEGIN_ARGS(HTML5DeprecationBanner)
+	{}
+	SLATE_END_ARGS()
+
+	void Construct(const FArguments& InArgs)
+	{
+		ChildSlot
+		[
+			SNew(SBorder)
+			.BorderBackgroundColor(this, &HTML5DeprecationBanner::GetBorderColor)
+			.BorderImage(FEditorStyle::GetBrush("ToolPanel.LightGroupBorder"))
+			.Padding(8.0f)
+			[
+				SNew(SHorizontalBox)
+				//	.ToolTipText(Tooltip)
+
+				// Status icon
+				+SHorizontalBox::Slot()
+				.AutoWidth()
+				.VAlign(VAlign_Center)
+				[
+					SNew(SImage)
+					.Image(FEditorStyle::GetBrush("SettingsEditor.WarningIcon"))
+				]
+
+				// Notice
+				+ SHorizontalBox::Slot()
+				.FillWidth(1.0f)
+				.Padding(16.0f, 0.0f)
+				.VAlign(VAlign_Center)
+				[
+					SNew(STextBlock)
+					.ColorAndOpacity(FLinearColor::White)
+					.ShadowColorAndOpacity(FLinearColor::Black)
+					.ShadowOffset(FVector2D::UnitVector)
+					.Text(LOCTEXT("HTML5DeprecatedPlatformMessageText",
+						"The HTML5 platform will be migrated from:\n\tan officialy supported (integrated) platform,\nand be moved to:\n\tas a community supported (plugin) platform."))
+				]
+			]
+		];
+	}
+
+	FSlateColor GetBorderColor() const
+	{
+		return FLinearColor(0.8f, 0, 0);
+	}
+};
+
+//////////////////////////////////////////////////////////////////////////
+
+
 TSharedRef<IDetailCustomization> FHTML5TargetSettingsCustomization::MakeInstance()
 {
 	return MakeShareable(new FHTML5TargetSettingsCustomization);
@@ -33,6 +101,14 @@ FHTML5TargetSettingsCustomization::~FHTML5TargetSettingsCustomization()
 }
 
 void FHTML5TargetSettingsCustomization::CustomizeDetails(IDetailLayoutBuilder& DetailLayout)
+{
+	BuildDeprecationMessage(DetailLayout); // !!! HTML5 Deprecation Message
+	// --------------------------------------------------
+	BuildEmscriptenCategory(DetailLayout); // !!! HTML5 Deprecation Message
+	AudioPluginWidgetManager.BuildAudioCategory(DetailLayout, EAudioPlatform::HTML5);
+}
+
+void FHTML5TargetSettingsCustomization::BuildEmscriptenCategory(IDetailLayoutBuilder& DetailLayout)
 {
 	IDetailCategoryBuilder& EmscriptenCategory = DetailLayout.EditCategory(TEXT("Emscripten"));
 
@@ -47,8 +123,17 @@ void FHTML5TargetSettingsCustomization::CustomizeDetails(IDetailLayoutBuilder& D
 	SETUP_SOURCEONLY_PROP(EnableIndexedDB, EmscriptenCategory);
 
 	SETUP_SOURCEONLY_PROP(EnableMultithreading, EmscriptenCategory);
+}
 
-	AudioPluginWidgetManager.BuildAudioCategory(DetailLayout, EAudioPlatform::HTML5);
+void FHTML5TargetSettingsCustomization::BuildDeprecationMessage(IDetailLayoutBuilder& DetailLayout)
+{
+	IDetailCategoryBuilder& HTML5PlatformDeprecatedCategory = DetailLayout.EditCategory(TEXT("HTML5 as a UE4 Platform Note:"));
+	TSharedRef<HTML5DeprecationBanner> PlatformSetupMessage = SNew(HTML5DeprecationBanner);
+	HTML5PlatformDeprecatedCategory.AddCustomRow(LOCTEXT("Warning", "Warning"), false)
+		.WholeRowWidget
+		[
+			PlatformSetupMessage
+		];
 }
 
 //////////////////////////////////////////////////////////////////////////
diff --git a/Engine/Source/Developer/HTML5/HTML5PlatformEditor/Private/HTML5TargetSettingsCustomization.h b/Engine/Source/Developer/HTML5/HTML5PlatformEditor/Private/HTML5TargetSettingsCustomization.h
--- a/Engine/Source/Developer/HTML5/HTML5PlatformEditor/Private/HTML5TargetSettingsCustomization.h
+++ b/Engine/Source/Developer/HTML5/HTML5PlatformEditor/Private/HTML5TargetSettingsCustomization.h
@@ -9,6 +9,8 @@
 #include "IDetailCustomization.h"
 #include "TargetPlatformAudioCustomization.h"
 
+#include "Widgets/SCompoundWidget.h"
+
 class IPropertyHandle;
 class SEditableTextBox;
 class SErrorText;
@@ -32,6 +34,9 @@ public:
 private:
 	FHTML5TargetSettingsCustomization();
 
+	void BuildDeprecationMessage(IDetailLayoutBuilder& DetailLayout);
+	void BuildEmscriptenCategory(IDetailLayoutBuilder& DetailLayout);
+
 	//Audio plugin widget builder:
 	FAudioPluginWidgetManager AudioPluginWidgetManager;
 };
