HTML5 - re add experimental tag to Multithreading Support checkbox in Project Settings for HTML5

#rnx

diff --git a/Engine/Source/Developer/HTML5/HTML5PlatformEditor/Classes/HTML5TargetSettings.h b/Engine/Source/Developer/HTML5/HTML5PlatformEditor/Classes/HTML5TargetSettings.h
--- a/Engine/Source/Developer/HTML5/HTML5PlatformEditor/Classes/HTML5TargetSettings.h
+++ b/Engine/Source/Developer/HTML5/HTML5PlatformEditor/Classes/HTML5TargetSettings.h
@@ -55,7 +55,7 @@ public:
 	/**
 	 * Enable Multithreading (Experimental)
 	 */
-	UPROPERTY(GlobalConfig, EditAnywhere, Category=Emscripten, Meta = (DisplayName = "Multithreading support"))
+	UPROPERTY(GlobalConfig, EditAnywhere, Category=Emscripten, Meta = (DisplayName = "Multithreading support [experimental]"))
 	bool EnableMultithreading;
 
 	/**
