HTML5 - deprecation messaging - spelling correction

#rnx

diff --git a/Engine/Source/Developer/HTML5/HTML5PlatformEditor/Private/HTML5TargetSettingsCustomization.cpp b/Engine/Source/Developer/HTML5/HTML5PlatformEditor/Private/HTML5TargetSettingsCustomization.cpp
--- a/Engine/Source/Developer/HTML5/HTML5PlatformEditor/Private/HTML5TargetSettingsCustomization.cpp
+++ b/Engine/Source/Developer/HTML5/HTML5PlatformEditor/Private/HTML5TargetSettingsCustomization.cpp
@@ -72,7 +72,7 @@ class HTML5DeprecationBanner : public SCompoundWidget
 					.ShadowColorAndOpacity(FLinearColor::Black)
 					.ShadowOffset(FVector2D::UnitVector)
 					.Text(LOCTEXT("HTML5DeprecatedPlatformMessageText",
-						"The HTML5 platform will be migrated from:\n\tan officialy supported (integrated) platform,\nand be moved to:\n\tas a community supported (plugin) platform."))
+						"The HTML5 platform will be migrated from:\n\tan officially supported (integrated) platform,\nand be moved to:\n\tas a community supported (plugin) platform."))
 				]
 			]
 		];
