HTML5 - updating deprecation messaging

#rnx

diff --git a/Engine/Source/Developer/HTML5/HTML5PlatformEditor/Private/HTML5TargetSettingsCustomization.cpp b/Engine/Source/Developer/HTML5/HTML5PlatformEditor/Private/HTML5TargetSettingsCustomization.cpp
--- a/Engine/Source/Developer/HTML5/HTML5PlatformEditor/Private/HTML5TargetSettingsCustomization.cpp
+++ b/Engine/Source/Developer/HTML5/HTML5PlatformEditor/Private/HTML5TargetSettingsCustomization.cpp
@@ -72,7 +72,8 @@ class HTML5DeprecationBanner : public SCompoundWidget
 					.ShadowColorAndOpacity(FLinearColor::Black)
 					.ShadowOffset(FVector2D::UnitVector)
 					.Text(LOCTEXT("HTML5DeprecatedPlatformMessageText",
-						"The HTML5 platform will be migrated from:\n\tan officially supported (integrated) platform,\nand be moved to:\n\tas a community supported (plugin) platform."))
+						"HTML5 platform support will be migrated to GitHub as a community-supported Platform Extension and no longer officially supported by Epic in upcoming releases."))
+					.AutoWrapText(true)
 				]
 			]
 		];
